Fixed some unlocalizable strings in categories for JA and CN

diff --git a/Templates/TemplateCategories.ini b/Templates/TemplateCategories.ini
--- a/Templates/TemplateCategories.ini
+++ b/Templates/TemplateCategories.ini
@@ -1,4 +1,4 @@
-; This file lists all the template categories that are exposed from this folder.
+﻿; This file lists all the template categories that are exposed from this folder.
 ; Add template categories here to have them show up in the Project Browser.
 ; Additional TemplateDefs.ini files can be created in your plugin root folder, at the same level as a Templates/ folder.
 ; Categories will contain project templates that include the category Key in their Categories property in their Config/TemplateDefs.ini.
@@ -12,10 +12,10 @@
 ; "Icon": A relative reference to the icon to display for the category - should be at minimum 128x128 to avoid being blurry.
 ; "IsMajorCategory": Is this a major, top-level category? Major categories are displayed as full rows, eg. the "Game" category.
 
-Categories=(Key="Games", LocalizedDisplayNames=((Language="en", Text="Games")), LocalizedDescriptions=((Language="en", Text="Start your game development journey with one of our key classes, levels, and examples.")), Icon="Media/Games.png", IsMajorCategory=true)
+Categories=(Key="Games", LocalizedDisplayNames=((Language="en", Text="Games"),(Language="zh-CN", Text="游戏"),(Language="JA", Text="ゲーム")), LocalizedDescriptions=((Language="en", Text="Start your game development journey with one of our key classes, levels, and examples."),(Language="zh-CN", Text="使用我们的关键职业、关卡和示例开始你的游戏开发之旅。"),(Language="JA", Text="基本となるクラスやレベル、サンプルを一つ使ってゲーム開発をはじめましょう。")), Icon="Media/Games.png", IsMajorCategory=true)
 
-Categories=(Key="ME", LocalizedDisplayNames=((Language="en", Text="Film, Television, and Live Events")), LocalizedDescriptions=((Language="en", Text="Choose from templates and examples for nDisplay, VR Scouting, and virtual production workflows.")), Icon="Media/ME.png", IsMajorCategory=true)
+Categories=(Key="ME", LocalizedDisplayNames=((Language="en", Text="Film, Television, and Live Events"),(Language="zh-CN", Text="影视与现场活动"),(Language="JA", Text="映画、テレビ、ライブ イベント")), LocalizedDescriptions=((Language="en", Text="Choose from templates and examples for nDisplay, VR Scouting, and virtual production workflows."),(Language="en", Text="选择适用于nDisplay、VR探查和虚拟制片工作流程的模板和示例。"),(Language="JA", Text="nDisplay やバーチャル スカウティング ツール、バーチャル プロダクション ワークフローのためのテンプレートおよびサンプルを選択することができます。")), Icon="Media/ME.png", IsMajorCategory=true)
 
-Categories=(Key="AEC", LocalizedDisplayNames=((Language="en", Text="Architecture, Engineering, and Construction")), LocalizedDescriptions=((Language="en", Text="Select a starting point for multi-user design reviews, photorealistic architectural design visualizations, sunlight studies, or stylized renderings.")), Icon="Media/AEC.png", IsMajorCategory=true)
+Categories=(Key="AEC", LocalizedDisplayNames=((Language="en", Text="Architecture, Engineering, and Construction"),(Language="zh-CN", Text="建筑、工程与施工"),(Language="JA", Text="建築、土木、建設")), LocalizedDescriptions=((Language="en", Text="Select a starting point for multi-user design reviews, photorealistic architectural design visualizations, sunlight studies, or stylized renderings."),(Language="zh-CN", Text="为多用户设计评审、照片级建筑设计可视化、阳光研究或风格化渲染选择一个起点。"),(Language="JA", Text="複数のユーザーによるレビューや、フォトリアルな建築デザインのビジュアライゼーション、太陽光の調査、スタイライズド レンダリングのための土台を選択することができます。")), Icon="Media/AEC.png", IsMajorCategory=true)
 
-Categories=(Key="MFG", LocalizedDisplayNames=((Language="en", Text="Automotive, Product Design, and Manufacturing")), LocalizedDescriptions=((Language="en", Text="Find templates for multi-user design reviews, photobooth studio environments, and product configurators.")),  Icon="Media/MFG.png", IsMajorCategory=true)
\ No newline at end of file
+Categories=(Key="MFG", LocalizedDisplayNames=((Language="en", Text="Automotive, Product Design, and Manufacturing"),(Language="zh-CN", Text="汽车、产品设计和制造"),(Language="JA", Text="自動車、プロダクト デザイン、製造")), LocalizedDescriptions=((Language="en", Text="Find templates for multi-user design reviews, photobooth studio environments, and product configurators."),(Language="zh-CN", Text="为多用户设计评审、Photobooth工作室环境和产品配置器寻找适合的模板。"),(Language="JA", Text="複数のユーザーによるレビューや、フォトブース スタジオの環境背景、製品コンフィギュレーターのためのテンプレートを見つけることができます。")),  Icon="Media/MFG.png", IsMajorCategory=true)
\ No newline at end of file
