Templates KO localization

diff --git a/Templates/TP_AEC_ArchvisBP/Config/TemplateDefs.ini b/Templates/TP_AEC_ArchvisBP/Config/TemplateDefs.ini
--- a/Templates/TP_AEC_ArchvisBP/Config/TemplateDefs.ini
+++ b/Templates/TP_AEC_ArchvisBP/Config/TemplateDefs.ini
@@ -1,4 +1,4 @@
-; This file specifies how a template will be treated when creating a project from it.
+﻿; This file specifies how a template will be treated when creating a project from it.
 ; It includes rules for ignoring files, replacing files and changing the contents of files.
 ; This file is annotated to help with creating your own templates.
 
@@ -12,6 +12,10 @@ LocalizedDisplayNames=(Language="en", Text="Archvis")
 
 ; The localized descriptions for this template, which is displayed when the template is selected.
 LocalizedDescriptions=(Language="en", Text="This template is dedicated to architectural visualization workflows with examples for sun studies, interior renderings, and non-photorealistic stylized renderings.\n\nThis template enables you to use Datasmith to import models from various design applications and easily light them.")
+
+LocalizedDisplayNames=(Language="ko",Text="건축 시각화")
+LocalizedDescriptions=(Language="ko",Text="이 템플릿은 조명 연구, 인테리어 렌더링 그리고 비실사 스타일라이즈드 렌더링을 위한 예제를 포함한 건축 시각화 워크플로를 위한 템플릿입니다.\n\n이 템플릿은 다양한 디자인 어플리케이션에서 데이터스미스를 이용한 모델 임포트와 시각화가 가능합니다.")
+
 ; The sort key to use for the template to bypass the default alphabetical sorting.
 SortKey="1"
 
diff --git a/Templates/TP_AEC_BlankBP/Config/TemplateDefs.ini b/Templates/TP_AEC_BlankBP/Config/TemplateDefs.ini
--- a/Templates/TP_AEC_BlankBP/Config/TemplateDefs.ini
+++ b/Templates/TP_AEC_BlankBP/Config/TemplateDefs.ini
@@ -1,4 +1,4 @@
-; This file specifies how a template will be treated when creating a project from it.
+﻿; This file specifies how a template will be treated when creating a project from it.
 ; It includes rules for ignoring files, replacing files and changing the contents of files.
 ; This file is annotated to help with creating your own templates.
 
@@ -13,6 +13,10 @@ LocalizedDisplayNames=(Language="en", Text="Blank")
 
 ; The localized descriptions for this template, which is displayed when the template is selected.
 LocalizedDescriptions=(Language="en", Text="Blank template with minimal functionality and content.\n\nThis template enables you to use Datasmith to import models from various design applications and easily light them.")
+
+LocalizedDisplayNames=(Language="ko", Text="빈 프로젝트")
+LocalizedDescriptions=(Language="ko", Text="최소한의 기능과 콘텐츠를 포함한 빈 프로젝트 템플릿\n\n이 템플릿은 다양한 디자인 어플리케이션에서 데이터스미스를 이용한 모델 임포트와 시각화가 가능합니다.")
+
 ; The sort key to use for the template to bypass the default alphabetical sorting.
 SortKey="0"
 
diff --git a/Templates/TP_CollaborativeBP/Config/TemplateDefs.ini b/Templates/TP_CollaborativeBP/Config/TemplateDefs.ini
--- a/Templates/TP_CollaborativeBP/Config/TemplateDefs.ini
+++ b/Templates/TP_CollaborativeBP/Config/TemplateDefs.ini
@@ -13,6 +13,10 @@ LocalizedDisplayNames=(Language="en", Text="Collab Viewer")
 
 ; The localized descriptions for this template, which is displayed when the template is selected.
 LocalizedDescriptions=(Language="en", Text="This project template provides navigation and interaction for various platforms (Desktop, VR). Use as a starting point for collaborative industrial or architectural projects.\n\nThis template enables you to use Datasmith to import models from various design applications and easily light them.")
+
+LocalizedDisplayNames=(Language="ko", Text="콜라보 뷰어")
+LocalizedDescriptions=(Language="ko", Text="이 프로젝트 템플릿은 데스크탑, VR등 다양한 플랫폼을 위한 네비게이션과 인터렉션을 제공합니다. 산업용 협업 또는 건축 프로젝트를 위한 시작 프로젝트로 사용하세요.")
+
 ; The sort key to use for the template to bypass the default alphabetical sorting.
 SortKey="3"
 
diff --git a/Templates/TP_CollaborativeBP/manifest.json b/Templates/TP_CollaborativeBP/manifest.json
--- a/Templates/TP_CollaborativeBP/manifest.json
+++ b/Templates/TP_CollaborativeBP/manifest.json
@@ -5,6 +5,10 @@
 		{
 			"Language": "en",
 			"Text": "Collab Viewer"
+		},
+		{
+			"Language": "ko",
+			"Text": "콜라보 뷰어"
 		}
 	],
 	"Description":
@@ -12,6 +16,10 @@
 		{
 			"Language": "en",
 			"Text": "This project template provides navigation and interaction for various platforms (Desktop, VR). Use as a starting point for collaborative Industrial or Architectural projects."
+		},
+		{
+			"Language": "ko",
+			"Text": "이 프로젝트 템플릿은 데스크탑, VR등 다양한 플랫폼을 위한 네비게이션과 인터렉션을 제공합니다. 산업용 협업 또는 건축 프로젝트를 위한 시작 프로젝트로 사용하세요."
 		}
 	],
 	"AssetTypes": [],
diff --git a/Templates/TP_ME_BlankBP/Config/TemplateDefs.ini b/Templates/TP_ME_BlankBP/Config/TemplateDefs.ini
--- a/Templates/TP_ME_BlankBP/Config/TemplateDefs.ini
+++ b/Templates/TP_ME_BlankBP/Config/TemplateDefs.ini
@@ -1,4 +1,4 @@
-; This file specifies how a template will be treated when creating a project from it.
+﻿; This file specifies how a template will be treated when creating a project from it.
 ; It includes rules for ignoring files, replacing files and changing the contents of files.
 ; This file is annotated to help with creating your own templates.
 
@@ -12,6 +12,10 @@ LocalizedDisplayNames=(Language="en", Text="Blank")
 
 ; The localized descriptions for this template, which is displayed when the template is selected.
 LocalizedDescriptions=(Language="en", Text="Blank template with minimal functionality.")
+
+LocalizedDisplayNames=(Language="ko", Text="빈 프로젝트")
+LocalizedDescriptions=(Language="ko", Text="최소한의 기능을 가진 빈 프로젝트 템플릿")
+
 ; The sort key to use for the template to bypass the default alphabetical sorting.
 SortKey="_1"
 
diff --git a/Templates/TP_ME_VProdBP/Config/TemplateDefs.ini b/Templates/TP_ME_VProdBP/Config/TemplateDefs.ini
--- a/Templates/TP_ME_VProdBP/Config/TemplateDefs.ini
+++ b/Templates/TP_ME_VProdBP/Config/TemplateDefs.ini
@@ -1,4 +1,4 @@
-; This file specifies how a template will be treated when creating a project from it.
+﻿; This file specifies how a template will be treated when creating a project from it.
 ; It includes rules for ignoring files, replacing files and changing the contents of files.
 ; This file is annotated to help with creating your own templates.
 
@@ -12,6 +12,10 @@ LocalizedDisplayNames=(Language="en", Text="Virtual Production")
 
 ; The localized descriptions for this template, which is displayed when the template is selected.
 LocalizedDescriptions=(Language="en", Text="This Virtual Production project template contains functionality for VR Scouting, Virtual Camera, SDI Video, Live Compositing, and nDisplay.")
+
+LocalizedDisplayNames=(Language="ko", Text="버추얼 프로덕션")
+LocalizedDescriptions=(Language="ko", Text="이 버추얼 프로덕션 프로젝트 템플릿은 VR 스카우팅, 버추얼 카메라, SDI 비디오, 라이브 컴포지팅 그리고 nDisplay를 위한 기능들이 포함되어 있습니다.")
+
 ; The sort key to use for the template to bypass the default alphabetical sorting.
 SortKey="_2"
 
diff --git a/Templates/TP_PhotoStudioBP/Config/TemplateDefs.ini b/Templates/TP_PhotoStudioBP/Config/TemplateDefs.ini
--- a/Templates/TP_PhotoStudioBP/Config/TemplateDefs.ini
+++ b/Templates/TP_PhotoStudioBP/Config/TemplateDefs.ini
@@ -1,4 +1,4 @@
-; This file specifies how a template will be treated when creating a project from it.
+﻿; This file specifies how a template will be treated when creating a project from it.
 ; It includes rules for ignoring files, replacing files and changing the contents of files.
 ; This file is annotated to help with creating your own templates.
 
@@ -12,6 +12,10 @@ LocalizedDisplayNames=(Language="en", Text="Photo Studio")
 
 ; The localized descriptions for this template, which is displayed when the template is selected.
 LocalizedDescriptions=(Language="en", Text="This template enables you to create product presentations and cinematics in a photographic studio environment.\n\nThis template enables you to use Datasmith to import models from various design applications and easily light them.")
+
+LocalizedDisplayNames=(Language="ko", Text="포토 스튜디오")
+LocalizedDescriptions=(Language="ko", Text="이 템플릿은 사진 스튜디오 환경에서 제품 시연과 시네마틱을 제작할 수 있도록 해줍니다.\n\n이 템플릿은 다양한 디자인 어플리케이션에서 데이터스미스를 이용한 모델 임포트와 시각화가 가능합니다.")
+
 ; The sort key to use for the template to bypass the default alphabetical sorting.
 SortKey="1"
 
diff --git a/Templates/TemplateCategories.ini b/Templates/TemplateCategories.ini
--- a/Templates/TemplateCategories.ini
+++ b/Templates/TemplateCategories.ini
@@ -12,10 +12,10 @@
 ; "Icon": A relative reference to the icon to display for the category - should be at minimum 128x128 to avoid being blurry.
 ; "IsMajorCategory": Is this a major, top-level category? Major categories are displayed as full rows, eg. the "Game" category.
 
-Categories=(Key="Games", LocalizedDisplayNames=((Language="en",Text="Games"),(Language="zh-Hans",Text="游戏"),(Language="ja",Text="ゲーム")), LocalizedDescriptions=((Language="en",Text="Start your game development journey with one of our key classes, levels, and examples."),(Language="zh-Hans", Text="使用我们的关键职业、关卡和示例开始你的游戏开发之旅。"),(Language="ja",Text="基本となるクラスやレベル、サンプルを一つ使ってゲーム開発をはじめましょう。")), Icon="Media/Games.png", IsMajorCategory=true)
+Categories=(Key="Games", LocalizedDisplayNames=((Language="en",Text="Games"),(Language="zh-Hans",Text="游戏"),(Language="ja",Text="ゲーム"),(Language="ko",Text="게임")), LocalizedDescriptions=((Language="en",Text="Start your game development journey with one of our key classes, levels, and examples."),(Language="zh-Hans", Text="使用我们的关键职业、关卡和示例开始你的游戏开发之旅。"),(Language="ja",Text="基本となるクラスやレベル、サンプルを一つ使ってゲーム開発をはじめましょう。"),(Language="ko",Text="언리얼 엔진 주요 클래스, 레벨 및 샘플을 활용해서 게임 개발을 시작해 보세요.")), Icon="Media/Games.png", IsMajorCategory=true)
 
-Categories=(Key="ME", LocalizedDisplayNames=((Language="en",Text="Film, Television, and Live Events"),(Language="zh-Hans",Text="影视与现场活动"),(Language="ja",Text="映画、テレビ、ライブ イベント")), LocalizedDescriptions=((Language="en",Text="Choose from templates and examples for nDisplay, VR Scouting, and virtual production workflows."),(Language="zh-Hans",Text="选择适用于nDisplay、VR探查和虚拟制片工作流程的模板和示例。"),(Language="ja", Text="nDisplay やバーチャル スカウティング ツール、バーチャル プロダクション ワークフローのためのテンプレートおよびサンプルを選択することができます。")), Icon="Media/ME.png", IsMajorCategory=true)
+Categories=(Key="ME", LocalizedDisplayNames=((Language="en",Text="Film, Television, and Live Events"),(Language="zh-Hans",Text="影视与现场活动"),(Language="ja",Text="映画、テレビ、ライブ イベント"),(Language="ko",Text="영화, TV 및 라이브 이벤트")), LocalizedDescriptions=((Language="en",Text="Choose from templates and examples for nDisplay, VR Scouting, and virtual production workflows."),(Language="zh-Hans",Text="选择适用于nDisplay、VR探查和虚拟制片工作流程的模板和示例。"),(Language="ja", Text="nDisplay やバーチャル スカウティング ツール、バーチャル プロダクション ワークフローのためのテンプレートおよびサンプルを選択することができます。"),(Language="ko",Text="nDisplay, VR 스카우팅 및 버추얼 프로덕션 워크플로를 위한 템플릿과 샘플을 선택해 보세요.")), Icon="Media/ME.png", IsMajorCategory=true)
 
-Categories=(Key="AEC", LocalizedDisplayNames=((Language="en",Text="Architecture, Engineering, and Construction"),(Language="zh-Hans",Text="建筑、工程与施工"),(Language="ja",Text="建築、土木、建設")), LocalizedDescriptions=((Language="en",Text="Select a starting point for multi-user design reviews, photorealistic architectural design visualizations, sunlight studies, or stylized renderings."),(Language="zh-Hans", Text="为多用户设计评审、照片级建筑设计可视化、阳光研究或风格化渲染选择一个起点。"),(Language="ja",Text="複数のユーザーによるレビューや、フォトリアルな建築デザインのビジュアライゼーション、太陽光の調査、スタイライズド レンダリングのための土台を選択することができます。")), Icon="Media/AEC.png", IsMajorCategory=true)
+Categories=(Key="AEC", LocalizedDisplayNames=((Language="en",Text="Architecture, Engineering, and Construction"),(Language="zh-Hans",Text="建筑、工程与施工"),(Language="ja",Text="建築、土木、建設"),(Language="ko",Text="건축, 엔지니어링 및 건설")), LocalizedDescriptions=((Language="en",Text="Select a starting point for multi-user design reviews, photorealistic architectural design visualizations, sunlight studies, or stylized renderings."),(Language="zh-Hans", Text="为多用户设计评审、照片级建筑设计可视化、阳光研究或风格化渲染选择一个起点。"),(Language="ja",Text="複数のユーザーによるレビューや、フォトリアルな建築デザインのビジュアライゼーション、太陽光の調査、スタイライズド レンダリングのための土台を選択することができます。"),(Language="ko",Text="멀티 유저 디자인 리뷰, 건축 디자인의 포토리얼한 시각화, 선라이트 스터디 또는 스타일라이즈드 렌더링을 시작해 보세요.")), Icon="Media/AEC.png", IsMajorCategory=true)
 
-Categories=(Key="MFG", LocalizedDisplayNames=((Language="en",Text="Automotive, Product Design, and Manufacturing"),(Language="zh-Hans",Text="汽车、产品设计和制造"),(Language="ja",Text="自動車、プロダクト デザイン、製造")), LocalizedDescriptions=((Language="en",Text="Find templates for multi-user design reviews, photobooth studio environments, and product configurators."),(Language="zh-Hans",Text="为多用户设计评审、Photobooth工作室环境和产品配置器寻找适合的模板。"),(Language="ja",Text="複数のユーザーによるレビューや、フォトブース スタジオの環境背景、製品コンフィギュレーターのためのテンプレートを見つけることができます。")), Icon="Media/MFG.png", IsMajorCategory=true)
\ No newline at end of file
+Categories=(Key="MFG", LocalizedDisplayNames=((Language="en",Text="Automotive, Product Design, and Manufacturing"),(Language="zh-Hans",Text="汽车、产品设计和制造"),(Language="ja",Text="自動車、プロダクト デザイン、製造"),(Language="ko",Text="자동차, 제품 디자인 및 제조")), LocalizedDescriptions=((Language="en",Text="Find templates for multi-user design reviews, photobooth studio environments, and product configurators."),(Language="zh-Hans",Text="为多用户设计评审、Photobooth工作室环境和产品配置器寻找适合的模板。"),(Language="ja",Text="複数のユーザーによるレビューや、フォトブース スタジオの環境背景、製品コンフィギュレーターのためのテンプレートを見つけることができます。"),(Language="ko",Text="멀티 유저 디자인 리뷰, 포토부스 스튜디오 환경 및 제품 컨피규레이터를 위한 템플릿을 살펴보세요.")), Icon="Media/MFG.png", IsMajorCategory=true)
