Removed invalid icon file. Will be fixed properly by Pierre-Felix Breton shortly.

diff --git a/Templates/TemplateCategories.ini b/Templates/TemplateCategories.ini
--- a/Templates/TemplateCategories.ini
+++ b/Templates/TemplateCategories.ini
@@ -12,7 +12,7 @@
 ; "Icon": A relative reference to the icon to display for the category - should be at minimum 128x128 to avoid being blurry.
 ; "IsMajorCategory": Is this a major, top-level category? Major categories are displayed as full rows, eg. the "Game" category.
 
-Categories=(Key="Games", LocalizedDisplayNames=((Language="en", Text="Games")), LocalizedDescriptions=((Language="en", Text="Starting off points giving you a few key classes, levels and examples to begin your game.")), Icon="GameProjectDialog.BlueprintIcon", IsMajorCategory=true)
-Categories=(Key="ME", LocalizedDisplayNames=((Language="en", Text="Media & Entertainment")), LocalizedDescriptions=((Language="en", Text="need text need text need text")), Icon="GameProjectDialog.BlueprintIcon", IsMajorCategory=true)
-Categories=(Key="AEC", LocalizedDisplayNames=((Language="en", Text="Architecture, Engineering & Construction")), LocalizedDescriptions=((Language="en", Text="need text need text need text")), Icon="GameProjectDialog.BlueprintIcon", IsMajorCategory=true)
-Categories=(Key="MFG", LocalizedDisplayNames=((Language="en", Text="Manufacturing, Automotive & Product Design")), LocalizedDescriptions=((Language="en", Text="need text need text need text")),  Icon="GameProjectDialog.BlueprintIcon", IsMajorCategory=true)
\ No newline at end of file
+Categories=(Key="Games", LocalizedDisplayNames=((Language="en", Text="Games")), LocalizedDescriptions=((Language="en", Text="Starting off points giving you a few key classes, levels and examples to begin your game.")), IsMajorCategory=true)
+Categories=(Key="ME", LocalizedDisplayNames=((Language="en", Text="Media & Entertainment")), LocalizedDescriptions=((Language="en", Text="need text need text need text")), IsMajorCategory=true)
+Categories=(Key="AEC", LocalizedDisplayNames=((Language="en", Text="Architecture, Engineering & Construction")), LocalizedDescriptions=((Language="en", Text="need text need text need text")), IsMajorCategory=true)
+Categories=(Key="MFG", LocalizedDisplayNames=((Language="en", Text="Manufacturing, Automotive & Product Design")), LocalizedDescriptions=((Language="en", Text="need text need text need text")), IsMajorCategory=true)
\ No newline at end of file
