Changing Media and Entertainment category based on request from Marketing as they are aligning the web design to be closely matching this.  It is now "Film, Television and Live Events"

diff --git a/Templates/TemplateCategories.ini b/Templates/TemplateCategories.ini
--- a/Templates/TemplateCategories.ini
+++ b/Templates/TemplateCategories.ini
@@ -14,7 +14,7 @@
 
 Categories=(Key="Games", LocalizedDisplayNames=((Language="en", Text="Games")), LocalizedDescriptions=((Language="en", Text="Start your game development journey with one of our key classes, levels, and examples.")), Icon="Media/Games.png", IsMajorCategory=true)
 
-Categories=(Key="ME", LocalizedDisplayNames=((Language="en", Text="Media and Entertainment")), LocalizedDescriptions=((Language="en", Text="Choose from templates and examples for nDisplay, VR Scouting, and virtual production workflows.")), Icon="Media/ME.png", IsMajorCategory=true)
+Categories=(Key="ME", LocalizedDisplayNames=((Language="en", Text="Film, Television and Live Events")), LocalizedDescriptions=((Language="en", Text="Choose from templates and examples for nDisplay, VR Scouting, and virtual production workflows.")), Icon="Media/ME.png", IsMajorCategory=true)
 
 Categories=(Key="AEC", LocalizedDisplayNames=((Language="en", Text="Architecture, Engineering and Construction")), LocalizedDescriptions=((Language="en", Text="Select a starting point for multi-user design reviews, photorealistic architectural design visualizations, sunlight studies, or stylized renderings.")), Icon="Media/AEC.png", IsMajorCategory=true)
 
