Updating Template Categories to reflect recent changes from marketing about how industries are nammed.

diff --git a/Templates/TemplateCategories.ini b/Templates/TemplateCategories.ini
--- a/Templates/TemplateCategories.ini
+++ b/Templates/TemplateCategories.ini
@@ -14,8 +14,8 @@
 
 Categories=(Key="Games", LocalizedDisplayNames=((Language="en", Text="Games")), LocalizedDescriptions=((Language="en", Text="Start your game development journey with one of our key classes, levels, and examples.")), Icon="Media/Games.png", IsMajorCategory=true)
 
-Categories=(Key="ME", LocalizedDisplayNames=((Language="en", Text="Film, Television and Live Events")), LocalizedDescriptions=((Language="en", Text="Choose from templates and examples for nDisplay, VR Scouting, and virtual production workflows.")), Icon="Media/ME.png", IsMajorCategory=true)
+Categories=(Key="ME", LocalizedDisplayNames=((Language="en", Text="Film, Television, and Live Events")), LocalizedDescriptions=((Language="en", Text="Choose from templates and examples for nDisplay, VR Scouting, and virtual production workflows.")), Icon="Media/ME.png", IsMajorCategory=true)
 
-Categories=(Key="AEC", LocalizedDisplayNames=((Language="en", Text="Architecture, Engineering and Construction")), LocalizedDescriptions=((Language="en", Text="Select a starting point for multi-user design reviews, photorealistic architectural design visualizations, sunlight studies, or stylized renderings.")), Icon="Media/AEC.png", IsMajorCategory=true)
+Categories=(Key="AEC", LocalizedDisplayNames=((Language="en", Text="Architecture, Engineering, and Construction")), LocalizedDescriptions=((Language="en", Text="Select a starting point for multi-user design reviews, photorealistic architectural design visualizations, sunlight studies, or stylized renderings.")), Icon="Media/AEC.png", IsMajorCategory=true)
 
-Categories=(Key="MFG", LocalizedDisplayNames=((Language="en", Text="Manufacturing, Automotive and Product Design")), LocalizedDescriptions=((Language="en", Text="Find templates for multi-user design reviews, photobooth studio environments, and product configurators.")),  Icon="Media/MFG.png", IsMajorCategory=true)
\ No newline at end of file
+Categories=(Key="MFG", LocalizedDisplayNames=((Language="en", Text="Automotive, Product Design, and Manufacturing")), LocalizedDescriptions=((Language="en", Text="Find templates for multi-user design reviews, photobooth studio environments, and product configurators.")),  Icon="Media/MFG.png", IsMajorCategory=true)
\ No newline at end of file
