Architecture name change

diff --git a/Templates/TemplateCategories.ini b/Templates/TemplateCategories.ini
--- a/Templates/TemplateCategories.ini
+++ b/Templates/TemplateCategories.ini
@@ -16,6 +16,6 @@ Categories=(Key="Games", LocalizedDisplayNames=((Language="en",Text="Games"),(La
 
 Categories=(Key="ME", LocalizedDisplayNames=((Language="en",Text="Film / Video &\nLive Events"),(Language="zh-Hans",Text="影视与现场活动"),(Language="ja",Text="映画、テレビ、ライブ イベント"),(Language="ko",Text="영화, TV 및 라이브 이벤트")), LocalizedDescriptions=((Language="en",Text="Choose from templates and examples for nDisplay, VR Scouting, and virtual production workflows."),(Language="zh-Hans",Text="选择适用于nDisplay、VR探查和虚拟制片工作流程的模板和示例。"),(Language="ja", Text="nDisplay やバーチャル スカウティング ツール、バーチャル プロダクション ワークフローのためのテンプレートおよびサンプルを選択することができます。"),(Language="ko",Text="nDisplay, VR 스카우팅 및 버추얼 프로덕션 워크플로를 위한 템플릿과 샘플을 선택해 보세요.")), Icon="Media/FilmVideoLive_2x.png", IsMajorCategory=true)
 
-Categories=(Key="AEC", LocalizedDisplayNames=((Language="en",Text="Architecture\nEngineering &\nConstruction"),(Language="zh-Hans",Text="建筑、工程与施工"),(Language="ja",Text="建築、土木、建設"),(Language="ko",Text="건축, 엔지니어링 및 건설")), LocalizedDescriptions=((Language="en",Text="Select a starting point for multi-user design reviews, photorealistic architectural design visualizations, sunlight studies, or stylized renderings."),(Language="zh-Hans", Text="为多用户设计评审、照片级建筑设计可视化、阳光研究或风格化渲染选择一个起点。"),(Language="ja",Text="複数のユーザーによるレビューや、フォトリアルな建築デザインのビジュアライゼーション、太陽光の調査、スタイライズド レンダリングのための土台を選択することができます。"),(Language="ko",Text="멀티 유저 디자인 리뷰, 건축 디자인의 포토리얼한 시각화, 선라이트 스터디 또는 스타일라이즈드 렌더링을 시작해 보세요.")), Icon="Media/AEC_2x.png", IsMajorCategory=true)
+Categories=(Key="AEC", LocalizedDisplayNames=((Language="en",Text="Architecture\nEngineering &\nConstruction"),(Language="zh-Hans",Text="建筑"),(Language="ja",Text="建築"),(Language="ko",Text="건축")), LocalizedDescriptions=((Language="en",Text="Select a starting point for multi-user design reviews, photorealistic architectural design visualizations, sunlight studies, or stylized renderings."),(Language="zh-Hans", Text="为多用户设计评审、照片级建筑设计可视化、阳光研究或风格化渲染选择一个起点。"),(Language="ja",Text="複数のユーザーによるレビューや、フォトリアルな建築デザインのビジュアライゼーション、太陽光の調査、スタイライズド レンダリングのための土台を選択することができます。"),(Language="ko",Text="멀티 유저 디자인 리뷰, 건축 디자인의 포토리얼한 시각화, 선라이트 스터디 또는 스타일라이즈드 렌더링을 시작해 보세요.")), Icon="Media/AEC_2x.png", IsMajorCategory=true)
 
 Categories=(Key="MFG", LocalizedDisplayNames=((Language="en",Text="Automotive\nProduct Design &\nManufacturing"),(Language="zh-Hans",Text="汽车、产品设计和制造"),(Language="ja",Text="自動車、プロダクト デザイン、製造"),(Language="ko",Text="자동차, 제품 디자인 및 제조")), LocalizedDescriptions=((Language="en",Text="Find templates for multi-user design reviews, photobooth studio environments, and product configurators."),(Language="zh-Hans",Text="为多用户设计评审、Photobooth工作室环境和产品配置器寻找适合的模板。"),(Language="ja",Text="複数のユーザーによるレビューや、フォトブース スタジオの環境背景、製品コンフィギュレーターのためのテンプレートを見つけることができます。"),(Language="ko",Text="멀티 유저 디자인 리뷰, 포토부스 스튜디오 환경 및 제품 컨피규레이터를 위한 템플릿을 살펴보세요.")), Icon="Media/AutomotiveProduct_2x.png", IsMajorCategory=true)
