Uncommenting simulation category.

diff --git a/Templates/TemplateCategories.ini b/Templates/TemplateCategories.ini
--- a/Templates/TemplateCategories.ini
+++ b/Templates/TemplateCategories.ini
@@ -20,4 +20,4 @@ Categories=(Key="AEC", LocalizedDisplayNames=((Language="en",Text="Architecture"
 
 Categories=(Key="MFG", LocalizedDisplayNames=((Language="en",Text="Automotive\nProduct Design &\nManufacturing"),(Language="zh-Hans",Text="汽车、产品设计和制造"),(Language="ja",Text="自動車、プロダクト デザイン、製造"),(Language="ko",Text="자동차, 제품 디자인 및 제조")), LocalizedDescriptions=((Language="en",Text="Find templates for multi-user design reviews, photobooth studio environments, and product configurators."),(Language="zh-Hans",Text="为多用户设计评审、Photobooth工作室环境和产品配置器寻找适合的模板。"),(Language="ja",Text="複数のユーザーによるレビューや、フォトブース スタジオの環境背景、製品コンフィギュレーターのためのテンプレートを見つけることができます。"),(Language="ko",Text="멀티 유저 디자인 리뷰, 포토부스 스튜디오 환경 및 제품 컨피규레이터를 위한 템플릿을 살펴보세요.")), Icon="Media/AutomotiveProduct_2x.png", IsMajorCategory=true)
 
-;Categories=(Key="SIM", LocalizedDisplayNames=((Language="en",Text="Simulation"),(Language="zh-Hans",Text="TODO_Simulation_zh"),(Language="ja",Text="TODO_Simulation_ja"),(Language="ko",Text="TODO_Simulation_ko")), LocalizedDescriptions=((Language="en",Text="Find templates for Outdoor georeferenced environments, First person, Driving, Multichannels and AR/VR projects."),(Language="zh-Hans",Text="TODO_SimulationDesc_"),(Language="ja",Text="TODO_SimulationDesc_ja"),(Language="ko",Text="TODO_SimulationDesc_ko")), Icon="Media/Simulation_2x.png", IsMajorCategory=true)
+Categories=(Key="SIM", LocalizedDisplayNames=((Language="en",Text="Simulation"),(Language="zh-Hans",Text="TODO_Simulation_zh"),(Language="ja",Text="TODO_Simulation_ja"),(Language="ko",Text="TODO_Simulation_ko")), LocalizedDescriptions=((Language="en",Text="Find templates for Outdoor georeferenced environments, First person, Driving, Multichannels and AR/VR projects."),(Language="zh-Hans",Text="TODO_SimulationDesc_"),(Language="ja",Text="TODO_SimulationDesc_ja"),(Language="ko",Text="TODO_SimulationDesc_ko")), Icon="Media/Simulation_2x.png", IsMajorCategory=true)
