Updating the translation to match the udpated source text

diff --git a/Templates/TemplateCategories.ini b/Templates/TemplateCategories.ini
--- a/Templates/TemplateCategories.ini
+++ b/Templates/TemplateCategories.ini
@@ -20,4 +20,4 @@ Categories=(Key="AEC", LocalizedDisplayNames=((Language="en",Text="Architecture"
 
 Categories=(Key="MFG", LocalizedDisplayNames=((Language="en",Text="Automotive\nProduct Design &\nManufacturing"),(Language="zh-Hans",Text="汽车、产品设计和制造"),(Language="ja",Text="自動車、プロダクト デザイン、製造"),(Language="ko",Text="자동차, 제품 디자인 및 제조")), LocalizedDescriptions=((Language="en",Text="Find templates for multi-user design reviews, photobooth studio environments, and product configurators."),(Language="zh-Hans",Text="为多用户设计评审、Photobooth工作室环境和产品配置器寻找适合的模板。"),(Language="ja",Text="複数のユーザーによるレビューや、フォトブース スタジオの環境背景、製品コンフィギュレーターのためのテンプレートを見つけることができます。"),(Language="ko",Text="멀티 유저 디자인 리뷰, 포토부스 스튜디오 환경 및 제품 컨피규레이터를 위한 템플릿을 살펴보세요.")), Icon="Media/AutomotiveProduct_2x.png", IsMajorCategory=true)
 
-Categories=(Key="SIM", LocalizedDisplayNames=((Language="en",Text="Simulation"),(Language="zh-Hans",Text="TODO_Simulation_zh"),(Language="ja",Text="TODO_Simulation_ja"),(Language="ko",Text="TODO_Simulation_ko")), LocalizedDescriptions=((Language="en",Text="Find templates for Outdoor georeferenced environments, First person, Driving, Multichannels and AR/VR projects."),(Language="zh-Hans",Text="TODO_SimulationDesc_"),(Language="ja",Text="TODO_SimulationDesc_ja"),(Language="ko",Text="TODO_SimulationDesc_ko")), Icon="Media/Simulation_2x.png", IsMajorCategory=true)
+Categories=(Key="SIM", LocalizedDisplayNames=((Language="en",Text="Simulation"),(Language="zh-Hans",Text="模拟"),(Language="ja",Text="シミュレーション"),(Language="ko",Text="시뮬레이션")), LocalizedDescriptions=((Language="en",Text="Find templates for Outdoor georeferenced environments, First person, Driving, Multichannels and AR/VR projects."),(Language="zh-Hans",Text="查看可用于室外地理参考场景、第一人称项目、驾驶类场景、多通道项目，以及AR/VR项目的模板。"),(Language="ja",Text="屋外ジオリファレンス環境、ファーストパーソン、運転、マルチチャンネル、AR/VR プロジェクト向けのテンプレートです。"),(Language="ko",Text="아웃도어 지오레퍼런스드 환경, 1인칭, 드라이빙, 멀티 채널 및 AR/VR 프로젝트를 위한 템플릿을 찾아보세요.")), Icon="Media/Simulation_2x.png", IsMajorCategory=true)
